[TraceTrimmer] Added initial implementation for TraceTrimmer tool.

#Synced-CL 49988361

diff --git a/Engine/Source/Programs/TraceTrimmer/Private/Common.h b/Engine/Source/Programs/TraceTrimmer/Private/Common.h
new file mode 100644
--- /dev/null
+++ b/Engine/Source/Programs/TraceTrimmer/Private/Common.h
@@ -0,0 +1,80 @@
+// Copyright Epic Games, Inc. All Rights Reserved.
+
+#pragma once
+
+#include "CoreTypes.h"
+
+// TraceAnalysis
+#include "Trace/DataStream.h"
+#include "Trace/StoreClient.h"
+
+// TraceTrimmer
+#include "Io.h"
+
+namespace UE
+{
+namespace TraceAnalyzer
+{
+
+////////////////////////////////////////////////////////////////////////////////////////////////////
+
+struct FFileDataStream : public UE::Trace::IInDataStream
+{
+	~FFileDataStream()
+	{
+		CloseFile(Handle);
+	}
+
+	virtual int32 Read(void* Data, uint32 Size) override
+	{
+		return FileRead(Handle, Data, Size);
+	}
+
+	FileHandle Handle;
+};
+
+////////////////////////////////////////////////////////////////////////////////////////////////////
+
+inline TUniquePtr<UE::Trace::IInDataStream> GetDataStream(int ArgC, TCHAR const* const* ArgV)
+{
+	FileHandle Input = -1;
+
+	if (ArgC < 2)
+	{
+		FileHandle StdIn = GetStdIn();
+		if (!IsTty(StdIn))
+		{
+			Input = StdIn;
+		}
+	}
+	else if (ArgC != 3)
+	{
+		Input = OpenFile(ArgV[1], false);
+	}
+
+	if (Input != -1)
+	{
+		FFileDataStream* Stream = new FFileDataStream();
+		Stream->Handle = Input;
+		return TUniquePtr<UE::Trace::IInDataStream>(Stream);
+	}
+
+	if (ArgC != 3)
+	{
+		return nullptr;
+	}
+
+	uint32 StorePort = FCString::Strtoui64(ArgV[1], nullptr, 10);
+	if (UE::Trace::FStoreClient* StoreClient = UE::Trace::FStoreClient::Connect(TEXT("127.0.0.1"), StorePort))
+	{
+		uint32 TraceId = FCString::Strtoui64(ArgV[2], nullptr, 10);
+		return StoreClient->ReadTrace(TraceId);
+	}
+
+	return nullptr;
+}
+
+////////////////////////////////////////////////////////////////////////////////////////////////////
+
+} // namespace TraceAnalyzer
+} // namespace UE
diff --git a/Engine/Source/Programs/TraceTrimmer/Private/Io.h b/Engine/Source/Programs/TraceTrimmer/Private/Io.h
new file mode 100644
--- /dev/null
+++ b/Engine/Source/Programs/TraceTrimmer/Private/Io.h
@@ -0,0 +1,192 @@
+// Copyright Epic Games, Inc. All Rights Reserved.
+
+#pragma once
+
+#include "CoreTypes.h"
+
+#if PLATFORM_DESKTOP
+
+#if PLATFORM_WINDOWS
+
+////////////////////////////////////////////////////////////////////////////////////////////////////
+
+#include "Windows/AllowWindowsPlatformTypes.h"
+#include <Windows.h>
+#include "Windows/HideWindowsPlatformTypes.h"
+
+typedef PTRINT FileHandle;
+constexpr FileHandle InvalidFileHandle = FileHandle(-1);
+
+////////////////////////////////////////////////////////////////////////////////////////////////////
+
+inline bool IsValidFileHandle(FileHandle Handle)
+{
+	return Handle >= 0;
+}
+
+////////////////////////////////////////////////////////////////////////////////////////////////////
+
+inline FileHandle GetStdOut()
+{
+	return FileHandle(GetStdHandle(STD_OUTPUT_HANDLE));
+}
+
+////////////////////////////////////////////////////////////////////////////////////////////////////
+
+inline FileHandle GetStdIn()
+{
+	return FileHandle(GetStdHandle(STD_INPUT_HANDLE));
+}
+
+////////////////////////////////////////////////////////////////////////////////////////////////////
+
+inline FileHandle GetStdErr()
+{
+	return FileHandle(GetStdHandle(STD_ERROR_HANDLE));
+}
+
+////////////////////////////////////////////////////////////////////////////////////////////////////
+
+inline bool IsTty(FileHandle Handle)
+{
+	return GetFileType(HANDLE(Handle)) == FILE_TYPE_CHAR;
+}
+
+////////////////////////////////////////////////////////////////////////////////////////////////////
+
+inline int32 FileRead(FileHandle Handle, void* Data, int Size)
+{
+	DWORD BytesIn = 0;
+	if (!ReadFile(HANDLE(Handle), Data, Size, &BytesIn, nullptr))
+	{
+		return -1;
+	}
+	return int32(BytesIn);
+}
+
+////////////////////////////////////////////////////////////////////////////////////////////////////
+
+inline int FileWrite(FileHandle Handle, const void* Data, int32 Size)
+{
+	DWORD BytesOut;
+	if (!WriteFile(HANDLE(Handle), Data, Size, &BytesOut, nullptr))
+	{
+		return -1;
+	}
+	return int32(BytesOut);
+}
+
+////////////////////////////////////////////////////////////////////////////////////////////////////
+
+inline FileHandle OpenFile(const TCHAR* Path, bool Writeable)
+{
+	HANDLE Out = CreateFile(Path,
+			Writeable ? GENERIC_WRITE : GENERIC_READ,
+			FILE_SHARE_READ,
+			nullptr,
+			Writeable ? CREATE_ALWAYS : OPEN_EXISTING,
+			0, nullptr);
+	return (Out == INVALID_HANDLE_VALUE) ? FileHandle(-1) : FileHandle(Out);
+}
+
+////////////////////////////////////////////////////////////////////////////////////////////////////
+
+inline void CloseFile(FileHandle Handle)
+{
+	CloseHandle(HANDLE(Handle));
+}
+
+////////////////////////////////////////////////////////////////////////////////////////////////////
+
+#endif // PLATFORM_WINDOWS
+
+#if PLATFORM_UNIX || PLATFORM_MAC
+
+////////////////////////////////////////////////////////////////////////////////////////////////////
+
+#include <fcntl.h>
+#include <unistd.h>
+
+typedef int FileHandle;
+constexpr FileHandle InvalidFileHandle = FileHandle(-1);
+
+////////////////////////////////////////////////////////////////////////////////////////////////////
+
+inline bool IsValidFileHandle(FileHandle Handle)
+{
+	return Handle >= 0;
+}
+
+////////////////////////////////////////////////////////////////////////////////////////////////////
+
+inline FileHandle GetStdOut()
+{
+	return STDOUT_FILENO;
+}
+
+////////////////////////////////////////////////////////////////////////////////////////////////////
+
+inline FileHandle GetStdIn()
+{
+	return STDIN_FILENO;
+}
+
+////////////////////////////////////////////////////////////////////////////////////////////////////
+
+inline FileHandle GetStdErr()
+{
+	return STDERR_FILENO;
+}
+
+////////////////////////////////////////////////////////////////////////////////////////////////////
+
+inline bool IsTty(FileHandle Handle)
+{
+	return isatty(Handle);
+}
+
+////////////////////////////////////////////////////////////////////////////////////////////////////
+
+inline int32 FileRead(FileHandle Handle, void* Data, int Size)
+{
+	return read(Handle, Data, Size);
+}
+
+////////////////////////////////////////////////////////////////////////////////////////////////////
+
+inline int FileWrite(FileHandle Handle, const void* Data, int32 Size)
+{
+	return write(Handle, Data, Size);
+}
+
+////////////////////////////////////////////////////////////////////////////////////////////////////
+
+inline FileHandle OpenFile(const TCHAR* Path, bool Writeable)
+{
+	char InnerPath[512];
+	for (ANSICHAR& Out : InnerPath)
+	{
+		Out = ANSICHAR(*Path++ & 0x7f);
+		if (Out == '\0')
+		{
+			break;
+		}
+	}
+
+	int Flags = Writeable ? O_CREAT|O_WRONLY|O_TRUNC : O_RDONLY;
+	int Mode = S_IRUSR|S_IWUSR|S_IRGRP|S_IWGRP|S_IROTH|S_IWOTH;
+	return open(InnerPath, Flags, Mode);
+}
+
+////////////////////////////////////////////////////////////////////////////////////////////////////
+
+inline void CloseFile(FileHandle Handle)
+{
+	close(Handle);
+}
+
+////////////////////////////////////////////////////////////////////////////////////////////////////
+
+#endif // PLATFORM_UNIX || PLATFORM_MAC
+
+#endif // PLATFORM_DESKTOP
diff --git a/Engine/Source/Programs/TraceTrimmer/Private/TraceTrimmer.cpp b/Engine/Source/Programs/TraceTrimmer/Private/TraceTrimmer.cpp
new file mode 100644
--- /dev/null
+++ b/Engine/Source/Programs/TraceTrimmer/Private/TraceTrimmer.cpp
@@ -0,0 +1,370 @@
+// Copyright Epic Games, Inc. All Rights Reserved.
+
+#include "TraceTrimmer.h"
+
+#include "Containers/AnsiString.h"
+#include "HAL/PlatformFile.h"
+#include "HAL/PlatformFileManager.h"
+#include "Misc/CommandLine.h"
+#include "RequiredProgramMainCPPInclude.h"
+#include "Templates/UniquePtr.h"
+
+#include "Io.h"
+#include "TraceWriterTests.h"
+
+// TraceAnalysis
+#include "Trace/DataStream.h"
+#include "Trace/OutDataStream.h"
+
+// TraceServices
+#include "TraceServices/TraceTrimmer.h"
+
+#include <stdio.h>
+#include <limits>
+
+DEFINE_LOG_CATEGORY(LogTraceTrimmer);
+
+IMPLEMENT_APPLICATION(TraceTrimmer, "TraceTrimmer");
+
+#ifndef TRACE_TRIMMER_HAS_TRACE_WRITER_TESTS
+#define TRACE_TRIMMER_HAS_TRACE_WRITER_TESTS 1
+#endif // TRACE_TRIMMER_HAS_TRACE_WRITER_TESTS
+
+namespace UE::TraceTrimmer
+{
+
+////////////////////////////////////////////////////////////////////////////////////////////////////
+
+static void PrintUsage()
+{
+	puts("Usage:");
+	puts("  TraceTrimmer [<in_trace_file> | -tracehost=... -traceid=...] [-o=... | -outTraceHost=...] [<options...>]");
+	puts("Where:");
+	puts("  <in_trace_file> : the input *.utrace file");
+	puts("  -tracehost=<host:port> : the host address of the input trace server");
+	puts("  -traceid=<id> : the id of the input trace (identified by the trace server)");
+	puts("  -o=<out_trace_file> : the output trace file; defaults to <in_trace_file>.out.utrace");
+	puts("  -outTraceHost=<out_host> : the host address of the output trace server");
+	puts("  -startTime=<time> : start time filter in seconds");
+	puts("  -endTime=<time> : end time filter in seconds");
+	puts("  -include=<events> : comma separated list of trace events to include");
+	puts("  -exclude=<events> : comma separated list of trace events to exclude");
+	puts("Note:");
+	puts("    The include/exclude events are specified as \"logger.event\" wildcard patterns, ex.: \"$trace.*,cpu.*\".");
+#if TRACE_TRIMMER_HAS_TRACE_WRITER_TESTS
+	puts("To run the TraceWriter tests:");
+	puts("  TraceTrimmer -tests[=<out_path>]");
+	puts("Where:");
+	puts("  <out_path> : the path where to write the generated trace files");
+#endif // TRACE_TRIMMER_HAS_TRACE_WRITER_TESTS
+}
+
+////////////////////////////////////////////////////////////////////////////////////////////////////
+
+struct FCmdLineParams
+{
+	FString InputTraceFile;
+	FString InputTraceServer;
+	uint32 InputTraceId = 0;
+
+	FString OutputTraceFile;
+	FString OutputTraceServer;
+
+	TraceServices::FTrimParameters TrimParams;
+};
+
+////////////////////////////////////////////////////////////////////////////////////////////////////
+
+#if TRACE_TRIMMER_HAS_TRACE_WRITER_TESTS
+static int32 TryRunTests(int32 ArgC, TCHAR* ArgV[])
+{
+	if (ArgC == 2 &&
+		(FCString::Stricmp(TEXT("-tests"), ArgV[1]) == 0 ||
+		 FCString::Strnicmp(TEXT("-tests="), ArgV[1], 7) == 0))
+	{
+		FString TestsOutputPath;
+		if (FCString::Strlen(ArgV[1]) > 7)
+		{
+			TestsOutputPath = ArgV[1] + 7;
+		}
+
+		IPlatformFile& PlatformFile = FPlatformFileManager::Get().GetPlatformFile();
+		FPaths::NormalizeDirectoryName(TestsOutputPath);
+		PlatformFile.CreateDirectoryTree(*TestsOutputPath);
+
+		UE::Trace::Tests::TraceWriterTest0(*(TestsOutputPath / TEXT("t0.utrace")));
+		UE::Trace::Tests::TraceWriterTest1(*(TestsOutputPath / TEXT("t1.utrace")));
+		UE::Trace::Tests::TraceWriterTest2a(*(TestsOutputPath / TEXT("t2a.utrace")));
+		UE::Trace::Tests::TraceWriterTest2b(*(TestsOutputPath / TEXT("t2b.utrace")));
+		UE::Trace::Tests::TraceWriterTest3(*(TestsOutputPath / TEXT("t3.utrace")));
+		UE::Trace::Tests::TraceWriterTestCpuProfiler(*(TestsOutputPath / TEXT("t4.utrace")));
+
+		return 1;
+	}
+
+	return 0;
+}
+#endif // TRACE_TRIMMER_HAS_TRACE_WRITER_TESTS
+
+////////////////////////////////////////////////////////////////////////////////////////////////////
+
+static int32 ParseCommandLine(FCmdLineParams& Parameters, int32 ArgC, TCHAR* ArgV[])
+{
+	if (ArgC < 2)
+	{
+		PrintUsage();
+		return 0;
+	}
+
+	int32 StartArgIndex = 1;
+	if (ArgV[1][0] != TEXT('-'))
+	{
+		Parameters.InputTraceFile = ArgV[1];
+		StartArgIndex = 2;
+	}
+	for (int32 ArgIndex = StartArgIndex; ArgIndex < ArgC; ++ArgIndex)
+	{
+		const TCHAR* Arg = ArgV[ArgIndex];
+
+		if (FCString::Strnicmp(TEXT("-tracehost="), Arg, 11) == 0)
+		{
+			if (!Parameters.InputTraceServer.IsEmpty())
+			{
+				LogTraceTrimmerError("The input trace server cannot be specified multiple times!");
+				return -1;
+			}
+			if (!Parameters.InputTraceFile.IsEmpty())
+			{
+				LogTraceTrimmerError("Cannot specify both an input trace file and an input trace server / trace id!");
+				return -1;
+			}
+			Parameters.InputTraceServer = Arg + 11;
+		}
+		else if (FCString::Strnicmp(TEXT("-traceid="), Arg, 9) == 0)
+		{
+			if (Parameters.InputTraceId != 0)
+			{
+				LogTraceTrimmerError("The input trace id cannot be specified multiple times!");
+				return -1;
+			}
+			if (!Parameters.InputTraceFile.IsEmpty())
+			{
+				LogTraceTrimmerError("Cannot specify both an input trace file and an input trace server / trace id!");
+				return -1;
+			}
+			if (FCString::Strnicmp(TEXT("0x"), Arg + 9, 2) == 0)
+			{
+				Parameters.InputTraceId = FParse::HexNumber(Arg + 11);
+			}
+			else
+			{
+				Parameters.InputTraceId = FCString::Atoi(Arg + 9);
+			}
+		}
+		else if (FCString::Strnicmp(TEXT("-o="), Arg, 3) == 0)
+		{
+			if (!Parameters.OutputTraceFile.IsEmpty())
+			{
+				LogTraceTrimmerError("The output trace file cannot be specified multiple times!");
+				return -1;
+			}
+			if (!Parameters.OutputTraceServer.IsEmpty())
+			{
+				LogTraceTrimmerError("Cannot specify both an output trace file and an output trace server!");
+				return -1;
+			}
+			Parameters.OutputTraceFile = Arg + 3;
+		}
+		else if (FCString::Strnicmp(TEXT("-outTraceHost="), Arg, 14) == 0)
+		{
+			if (!Parameters.OutputTraceServer.IsEmpty())
+			{
+				LogTraceTrimmerError("The output trace server cannot be specified multiple times!");
+				return -1;
+			}
+			if (!Parameters.OutputTraceFile.IsEmpty())
+			{
+				LogTraceTrimmerError("Cannot specify both an output trace file and an output trace server!");
+				return -1;
+			}
+			Parameters.OutputTraceServer = Arg + 14;
+		}
+		else if (FCString::Strnicmp(TEXT("-startTime="), Arg, 11) == 0)
+		{
+			Parameters.TrimParams.StartTime = FCString::Atof(Arg + 11);
+		}
+		else if (FCString::Strnicmp(TEXT("-endTime="), Arg, 9) == 0)
+		{
+			Parameters.TrimParams.EndTime = FCString::Atof(Arg + 9);
+		}
+		else if (FCString::Strnicmp(TEXT("-include="), Arg, 9) == 0)
+		{
+			Parameters.TrimParams.Include = FAnsiString(Arg + 9);
+		}
+		else if (FCString::Strnicmp(TEXT("-exclude="), Arg, 9) == 0)
+		{
+			Parameters.TrimParams.Exclude = FAnsiString(Arg + 9);
+		}
+		else
+		{
+			LogTraceTrimmerWarning("Unknown cmd line argument '%s'!", Arg);
+		}
+	}
+
+	return 1;
+}
+
+////////////////////////////////////////////////////////////////////////////////////////////////////
+
+TUniquePtr<UE::Trace::IInDataStream> CreateInputStream(FCmdLineParams& Parameters)
+{
+	FileHandle Input = -1;
+	if (!Parameters.InputTraceFile.IsEmpty())
+	{
+		Input = OpenFile(*Parameters.InputTraceFile, false);
+		if (Input < 0)
+		{
+			LogTraceTrimmerError("Cannot open the input trace file ('%s')!", *Parameters.InputTraceFile);
+			return nullptr;
+		}
+	}
+
+	class FInputDataStream : public UE::Trace::IInDataStream
+	{
+	public:
+		FInputDataStream(FileHandle InHandle)
+			: Handle(InHandle)
+		{
+		}
+
+		virtual ~FInputDataStream()
+		{
+			Close();
+		}
+
+		virtual int32 Read(void* Data, uint32 Size) override
+		{
+			return FileRead(Handle, Data, Size);
+		}
+
+		virtual void Close() override
+		{
+			if (Handle >= 0)
+			{
+				CloseFile(Handle);
+				Handle = -1;
+			}
+		}
+
+		virtual bool WaitUntilReady() override
+		{
+			return true;
+		}
+
+	private:
+		FileHandle Handle = -1;
+	};
+
+	TUniquePtr<UE::Trace::IInDataStream> InputDataStream = MakeUnique<FInputDataStream>(Input);
+	return InputDataStream;
+}
+
+////////////////////////////////////////////////////////////////////////////////////////////////////
+
+TUniquePtr<UE::Trace::IOutDataStream> CreateOutputStream(FCmdLineParams& Parameters)
+{
+	FileHandle Output = -1;
+	if (!Parameters.OutputTraceFile.IsEmpty())
+	{
+		Output = OpenFile(*Parameters.OutputTraceFile, true);
+		if (Output < 0)
+		{
+			LogTraceTrimmerError("Cannot write the output trace file ('%s')!", *Parameters.OutputTraceFile);
+			return nullptr;
+		}
+	}
+
+	class FOutputDataStream : public UE::Trace::IOutDataStream
+	{
+	public:
+		FOutputDataStream(FileHandle InHandle)
+			: Handle(InHandle)
+		{
+		}
+
+		virtual ~FOutputDataStream()
+		{
+			Close();
+		}
+
+		virtual int32 Write(void* Data, uint32 Size) override
+		{
+			return FileWrite(Handle, Data, Size);
+		}
+
+		virtual void Close() override
+		{
+			if (Handle >= 0)
+			{
+				CloseFile(Handle);
+				Handle = -1;
+			}
+		}
+
+		virtual bool WaitUntilReady() override
+		{
+			return true;
+		}
+
+	private:
+		FileHandle Handle = -1;
+	};
+
+	TUniquePtr<UE::Trace::IOutDataStream> OutputDataStream = MakeUnique<FOutputDataStream>(Output);
+	return OutputDataStream;
+}
+
+////////////////////////////////////////////////////////////////////////////////////////////////////
+
+} // namespace UE::TraceTrimmer
+
+////////////////////////////////////////////////////////////////////////////////////////////////////
+
+INT32_MAIN_INT32_ARGC_TCHAR_ARGV()
+{
+	using namespace UE::TraceTrimmer;
+	int32 Result;
+
+#if TRACE_TRIMMER_HAS_TRACE_WRITER_TESTS
+	Result = TryRunTests(ArgC, ArgV);
+	if (Result != 0)
+	{
+		return Result;
+	}
+#endif // TRACE_TRIMMER_HAS_TRACE_WRITER_TESTS
+
+	UE::TraceTrimmer::FCmdLineParams Parameters;
+
+	Result = ParseCommandLine(Parameters, ArgC, ArgV);
+	if (Result <= 0)
+	{
+		return Result;
+	}
+
+	TUniquePtr<UE::Trace::IInDataStream> InputDataStream = CreateInputStream(Parameters);
+	if (!InputDataStream)
+	{
+		return -1;
+	}
+
+	TUniquePtr<UE::Trace::IOutDataStream> OutputDataStream = CreateOutputStream(Parameters);
+	if (!OutputDataStream)
+	{
+		return -1;
+	}
+
+	return TraceServices::Trim(Parameters.TrimParams, *InputDataStream, *OutputDataStream);
+}
+
+////////////////////////////////////////////////////////////////////////////////////////////////////
+////////////////////////////////////////////////////////////////////////////////////////////////////
diff --git a/Engine/Source/Programs/TraceTrimmer/Private/TraceTrimmer.h b/Engine/Source/Programs/TraceTrimmer/Private/TraceTrimmer.h
new file mode 100644
--- /dev/null
+++ b/Engine/Source/Programs/TraceTrimmer/Private/TraceTrimmer.h
@@ -0,0 +1,24 @@
+// Copyright Epic Games, Inc. All Rights Reserved.
+
+#pragma once
+
+#include "CoreTypes.h"
+#include "Logging/LogCategory.h"
+
+DECLARE_LOG_CATEGORY_EXTERN(LogTraceTrimmer, Log, All);
+
+#define LogTraceTrimmerError(Format, ...) \
+	{ \
+		wprintf(L"Error: " Format "\n", ##__VA_ARGS__); \
+		UE_LOG(LogTraceTrimmer, Error, TEXT(Format), ##__VA_ARGS__); \
+	}
+#define LogTraceTrimmerWarning(Format, ...) \
+	{ \
+		wprintf(L"Warning: " Format "\n", ##__VA_ARGS__); \
+		UE_LOG(LogTraceTrimmer, Warning, TEXT(Format), ##__VA_ARGS__); \
+	}
+#define LogTraceTrimmerMessage(Format, ...) \
+	{ \
+		wprintf(L"" Format "\n", ##__VA_ARGS__); \
+		UE_LOG(LogTraceTrimmer, Log, TEXT(Format), ##__VA_ARGS__); \
+	}
diff --git a/Engine/Source/Programs/TraceTrimmer/Private/TraceWriterTests.cpp b/Engine/Source/Programs/TraceTrimmer/Private/TraceWriterTests.cpp
new file mode 100644
--- /dev/null
+++ b/Engine/Source/Programs/TraceTrimmer/Private/TraceWriterTests.cpp
@@ -0,0 +1,1132 @@
+// Copyright Epic Games, Inc. All Rights Reserved.
+
+#include "TraceWriterTests.h"
+
+//#include "BuildSettings.h"
+#include "HAL/PlatformTime.h"
+#include "Misc/App.h"
+#include "Misc/Guid.h"
+#include "ProfilingDebugging/MiscTrace.h"
+
+// TraceAnalysis
+#include "Trace/TraceWriter.h"
+#include "Trace/OutDataStream.h"
+
+// TraceTrimmer
+#include "TraceTrimmer.h"
+
+UE_DISABLE_OPTIMIZATION
+
+namespace UE::Trace::Tests
+{
+
+////////////////////////////////////////////////////////////////////////////////////////////////////
+
+static uint64 GTestStartTime = 0;
+
+static void BeginTest(const TCHAR* TraceFile)
+{
+	LogTraceTrimmerMessage("Writing %s...", TraceFile);
+
+	GTestStartTime = FPlatformTime::Cycles64();
+}
+
+////////////////////////////////////////////////////////////////////////////////////////////////////
+
+static void EndTest(FTraceWriter& Writer)
+{
+	const uint64 TestDuration = FPlatformTime::Cycles64() - GTestStartTime;
+	const double TestDurationSeconds = double(TestDuration) * FPlatformTime::GetSecondsPerCycle64();
+
+	while (true)
+	{
+		FString Err = Writer.GetLastError();
+		if (Err.IsEmpty())
+		{
+			break;
+		}
+		LogTraceTrimmerError("%s", *Err);
+	}
+
+	LogTraceTrimmerMessage("  done (%.6f seconds)", TestDurationSeconds);
+}
+
+////////////////////////////////////////////////////////////////////////////////////////////////////
+
+void TraceWriterTest0(const TCHAR* TraceFile)
+{
+	BeginTest(TraceFile);
+
+	FFileOutDataStream DataStream;
+	if (!DataStream.Open(TraceFile))
+	{
+		return;
+	}
+
+	FTraceWriter Writer(DataStream);
+
+	Writer.Begin();
+	Writer.End();
+
+	EndTest(Writer);
+}
+
+////////////////////////////////////////////////////////////////////////////////////////////////////
+
+static void WriteDiagnosticsSession2Event(FTraceWriter& Writer)
+{
+	constexpr int InstanceIdSize = 4;
+	FGuid InstanceGuid = FApp::GetInstanceId();
+	uint32 InstanceId[InstanceIdSize];
+	for (int Index = 0; Index < InstanceIdSize; ++Index)
+	{
+		InstanceId[Index] = InstanceGuid[Index];
+	}
+
+	FAnsiStringView Platform = UE_STRINGIZE(UBT_COMPILED_PLATFORM);
+	FAnsiStringView AppName = UE_APP_NAME;
+	FStringView ProjectName = TEXTVIEW("");
+	const TCHAR* EngineVersion = TEXT("5.8"); // BuildSettings::GetEngineVersionString();
+
+	// Tests that the writer declares a valid Diagnostics.Session2 event.
+	const uint32 DiagnosticsSession2EventId = Writer.DeclareDiagnosticsSession2Event();
+	const FTraceWriterEventInfo* EventInfo = Writer.GetEventInfo(DiagnosticsSession2EventId);
+	check(EventInfo != nullptr);
+	check(EnumHasAllFlags(EventInfo->GetFlags(), ETraceWriterEventFlags::ImportantNoSync));
+
+	Writer.SetCurrentThreadImportants();
+	Writer.WriteEvent(DiagnosticsSession2EventId)
+		.Field(ANSITEXTVIEW("ConfigurationType"), uint8(EBuildConfiguration::Test)) // uint8(FApp::GetBuildConfiguration())
+		.Field(ANSITEXTVIEW("TargetType"), uint8(EBuildTargetType::Program)) // uint8(FApp::GetBuildTargetType())
+		.Field(ANSITEXTVIEW("Changelist"), 0) // BuildSettings::GetCurrentChangelist()
+		.Field(ANSITEXTVIEW("InstanceId"), TConstArrayView<uint32>(InstanceId, InstanceIdSize))
+		.Field(ANSITEXTVIEW("Platform"), Platform)
+		.Field(ANSITEXTVIEW("AppName"), AppName)
+		//.Field(ANSITEXTVIEW("ProjectName"), ProjectName)
+		//.Field(ANSITEXTVIEW("Branch"), TEXTVIEW(""))
+		//.Field(ANSITEXTVIEW("BuildVersion"), TEXTVIEW(""))
+		.Field(ANSITEXTVIEW("EngineVersion"), EngineVersion)
+		.Field(ANSITEXTVIEW("CommandLine"), TEXTVIEW("TraceTrimmer.exe -tests"))
+		//.Field(ANSITEXTVIEW("VFSPaths"), ANSITEXTVIEW("")) // BuildSettings::GetVfsPaths();
+		.End();
+}
+
+////////////////////////////////////////////////////////////////////////////////////////////////////
+// FMiscWriter
+////////////////////////////////////////////////////////////////////////////////////////////////////
+
+class FMiscWriter
+{
+public:
+	FMiscWriter(FTraceWriter& InWriter)
+		: Writer(InWriter)
+	{
+	}
+	virtual ~FMiscWriter()
+	{
+	}
+
+	void DeclareEvents()
+	{
+		DeclareFrameEvents();
+	}
+	void DeclareEvents2()
+	{
+		DeclareFrameEvents2();
+	}
+
+	void DeclareFrameEvents();
+	void DeclareFrameEvents2();
+
+	void WriteBeginFrameEvent(uint64 Cycle, uint8 FrameType);
+	void WriteBeginFrameEvent2(uint64 Cycle, uint8 FrameType);
+	void WriteEndFrameEvent(uint64 Cycle, uint8 FrameType);
+	void WriteEndFrameEvent2(uint64 Cycle, uint8 FrameType);
+
+	void Test();
+	void Test2();
+
+private:
+	FTraceWriter& Writer;
+	uint32 EventId_Misc_BeginFrame = 0;
+	uint32 EventId_Misc_EndFrame = 0;
+};
+
+////////////////////////////////////////////////////////////////////////////////////////////////////
+
+void FMiscWriter::DeclareFrameEvents()
+{
+	EventId_Misc_BeginFrame =
+	Writer.DeclareEvent(ANSITEXTVIEW("Misc"), ANSITEXTVIEW("BeginFrame"))
+		.Field(ANSITEXTVIEW("Cycle"), ETraceWriterFieldType::Uint64)
+		.Field(ANSITEXTVIEW("FrameType"), ETraceWriterFieldType::Uint8)
+		.End();
+
+	EventId_Misc_EndFrame =
+	Writer.DeclareEvent(ANSITEXTVIEW("Misc"), ANSITEXTVIEW("EndFrame"))
+		.Field(ANSITEXTVIEW("Cycle"), ETraceWriterFieldType::Uint64)
+		.Field(ANSITEXTVIEW("FrameType"), ETraceWriterFieldType::Uint8)
+		.End();
+}
+
+////////////////////////////////////////////////////////////////////////////////////////////////////
+
+void FMiscWriter::DeclareFrameEvents2()
+{
+	EventId_Misc_BeginFrame =
+	UE_ADHOC_TRACE_BEGIN_DECLARE_EVENT(Writer, "Misc", "BeginFrame", None) // Sync
+		UE_ADHOC_TRACE_DECLARE_FIELD(uint64, "Cycle")
+		UE_ADHOC_TRACE_DECLARE_FIELD(uint8, "FrameType")
+	UE_ADHOC_TRACE_END_DECLARE_EVENT
+
+	EventId_Misc_EndFrame =
+	UE_ADHOC_TRACE_BEGIN_DECLARE_EVENT(Writer, "Misc", "EndFrame", None) // Sync
+		UE_ADHOC_TRACE_DECLARE_FIELD(uint64, "Cycle")
+		UE_ADHOC_TRACE_DECLARE_FIELD(uint8, "FrameType")
+	UE_ADHOC_TRACE_END_DECLARE_EVENT
+
+#if 0 // TBD !!!
+	EventId_Misc_BeginFrame =
+	UE_ADHOC_TRACE_BEGIN_DECLARE_EVENT(Writer, Misc, BeginFrame, Sync)
+		UE_ADHOC_TRACE_DECLARE_FIELD(uint64, Cycle)
+		UE_ADHOC_TRACE_DECLARE_FIELD(uint8, FrameType)
+	UE_ADHOC_TRACE_END_DECLARE_EVENT
+		
+	EventId_Misc_EndFrame =
+	UE_ADHOC_TRACE_BEGIN_DECLARE_EVENT(Writer, Misc, EndFrame, Sync)
+		UE_ADHOC_TRACE_DECLARE_FIELD(uint64, Cycle)
+		UE_ADHOC_TRACE_DECLARE_FIELD(uint8, FrameType)
+	UE_ADHOC_TRACE_END_DECLARE_EVENT()
+#endif
+}
+
+////////////////////////////////////////////////////////////////////////////////////////////////////
+
+void FMiscWriter::WriteBeginFrameEvent(uint64 Cycle, uint8 FrameType)
+{
+	Writer.WriteEvent(EventId_Misc_BeginFrame)
+		.Field(0, Cycle)
+		.Field(1, FrameType)
+		.End();
+}
+
+////////////////////////////////////////////////////////////////////////////////////////////////////
+
+void FMiscWriter::WriteBeginFrameEvent2(uint64 Cycle, uint8 FrameType)
+{
+	UE_ADHOC_TRACE_LOG(Writer, EventId_Misc_BeginFrame, Cycle, FrameType);
+
+	//UE_ADHOC_TRACE_WRITE_EVENT(Writer, EventId_Misc_BeginFrame, Cycle, FrameType); // TBD !!!
+
+	// TBD !!!
+	//UE_ADHOC_TRACE_BEGIN_WRITE_EVENT(Writer, EventId_Misc_BeginFrame)
+	//	UE_ADHOC_TRACE_WRITE_FIELD(0, Cycle)
+	//	UE_ADHOC_TRACE_WRITE_FIELD(1, FrameType)
+	//UE_ADHOC_TRACE_END_WRITE_EVENT()
+
+	// TBD !!!
+	//UE_ADHOC_TRACE_BEGIN_WRITE_EVENT(Writer, EventId_Misc_BeginFrame)
+	//	UE_ADHOC_TRACE_WRITE_FIELD_BY_NAME("Cycle", Cycle)
+	//	UE_ADHOC_TRACE_WRITE_FIELD_BY_NAME("FrameType", FrameType)
+	//UE_ADHOC_TRACE_END_WRITE_EVENT()
+}
+
+////////////////////////////////////////////////////////////////////////////////////////////////////
+
+void FMiscWriter::WriteEndFrameEvent(uint64 Cycle, uint8 FrameType)
+{
+	Writer.WriteEvent(EventId_Misc_EndFrame)
+		.Field(0, Cycle)
+		.Field(1, FrameType)
+		.End();
+}
+
+////////////////////////////////////////////////////////////////////////////////////////////////////
+
+void FMiscWriter::WriteEndFrameEvent2(uint64 Cycle, uint8 FrameType)
+{
+	UE_ADHOC_TRACE_LOG(Writer, EventId_Misc_EndFrame, Cycle, FrameType);
+}
+
+////////////////////////////////////////////////////////////////////////////////////////////////////
+
+void FMiscWriter::Test()
+{
+	const uint32 GameThreadId = Writer.RegisterThread(ANSITEXTVIEW("GameThread"));
+	const uint32 RenderThreadId = Writer.RegisterThread(ANSITEXTVIEW("RenderThread"));
+
+	DeclareEvents();
+
+	const uint64 BaseTimestamp = Writer.GetStartTime();
+	const uint64 OneSecond = Writer.GetTimeFrequency();
+
+	constexpr uint64 NumPackets = 10;
+	constexpr uint64 FramesPerPacket = 100;
+
+	const uint64 MaxFrameDuration = (16 * OneSecond) / 1000;
+	FMath::RandInit(42);
+
+	for (uint64 PacketIndex = 0; PacketIndex < NumPackets; ++PacketIndex)
+	{
+		Writer.SetCurrentThread(GameThreadId);
+		for (uint64 FrameIndex = 0; FrameIndex < FramesPerPacket; ++FrameIndex)
+		{
+			const uint64 Timestamp = BaseTimestamp + MaxFrameDuration * (PacketIndex * FramesPerPacket + FrameIndex);
+			const uint64 Duration = 6 * MaxFrameDuration / 7 + MaxFrameDuration * (uint64)FMath::Rand() / ((uint64)RAND_MAX * 7);
+
+			WriteBeginFrameEvent(Timestamp, (uint8)TraceFrameType_Game);
+			WriteEndFrameEvent(Timestamp + Duration, (uint8)TraceFrameType_Game);
+		}
+
+		Writer.SetCurrentThread(RenderThreadId);
+		for (uint64 FrameIndex = 0; FrameIndex < FramesPerPacket; ++FrameIndex)
+		{
+			const uint64 Timestamp = BaseTimestamp + MaxFrameDuration / 10 + MaxFrameDuration * (PacketIndex * FramesPerPacket + FrameIndex);
+			const uint64 Duration = 7 * MaxFrameDuration / 8 + MaxFrameDuration * (uint64)FMath::Rand() / ((uint64)RAND_MAX * 8);
+
+			WriteBeginFrameEvent(Timestamp, (uint8)TraceFrameType_Rendering);
+			WriteEndFrameEvent(Timestamp + Duration, (uint8)TraceFrameType_Rendering);
+		}
+	}
+}
+
+////////////////////////////////////////////////////////////////////////////////////////////////////
+
+void FMiscWriter::Test2()
+{
+	const uint32 GameThreadId = Writer.RegisterThread("GameThread");
+	const uint32 RenderThreadId = Writer.RegisterThread("RenderThread");
+
+	DeclareEvents2();
+
+	const uint64 BaseTimestamp = Writer.GetStartTime();
+	const uint64 OneSecond = Writer.GetTimeFrequency();
+
+	constexpr uint64 NumPackets = 10;
+	constexpr uint64 FramesPerPacket = 100;
+
+	const uint64 MaxFrameDuration = (16 * OneSecond) / 1000;
+	FMath::RandInit(42);
+
+	for (uint64 PacketIndex = 0; PacketIndex < NumPackets; ++PacketIndex)
+	{
+		Writer.SetCurrentThread(GameThreadId);
+		for (uint64 FrameIndex = 0; FrameIndex < FramesPerPacket; ++FrameIndex)
+		{
+			const uint64 Timestamp = BaseTimestamp + MaxFrameDuration * (PacketIndex * FramesPerPacket + FrameIndex);
+			const uint64 Duration = 6 * MaxFrameDuration / 7 + MaxFrameDuration * (uint64)FMath::Rand() / ((uint64)RAND_MAX * 7);
+
+			WriteBeginFrameEvent(Timestamp, (uint8)TraceFrameType_Game);
+			WriteEndFrameEvent(Timestamp + Duration, (uint8)TraceFrameType_Game);
+		}
+
+		Writer.SetCurrentThread(RenderThreadId);
+		for (uint64 FrameIndex = 0; FrameIndex < FramesPerPacket; ++FrameIndex)
+		{
+			const uint64 Timestamp = BaseTimestamp + MaxFrameDuration / 10 + MaxFrameDuration * (PacketIndex * FramesPerPacket + FrameIndex);
+			const uint64 Duration = 7 * MaxFrameDuration / 8 + MaxFrameDuration * (uint64)FMath::Rand() / ((uint64)RAND_MAX * 8);
+
+			WriteBeginFrameEvent2(Timestamp, (uint8)TraceFrameType_Rendering);
+			WriteEndFrameEvent2(Timestamp + Duration, (uint8)TraceFrameType_Rendering);
+		}
+	}
+}
+
+////////////////////////////////////////////////////////////////////////////////////////////////////
+////////////////////////////////////////////////////////////////////////////////////////////////////
+
+void TraceWriterTest1(const TCHAR* TraceFile)
+{
+	BeginTest(TraceFile);
+
+	FFileOutDataStream DataStream;
+	if (!DataStream.Open(TraceFile))
+	{
+		return;
+	}
+
+	FTraceWriter Writer(DataStream);
+	Writer.Begin();
+	WriteDiagnosticsSession2Event(Writer);
+	Writer.End();
+
+	EndTest(Writer);
+}
+
+////////////////////////////////////////////////////////////////////////////////////////////////////
+
+void TraceWriterTest2a(const TCHAR* TraceFile)
+{
+	BeginTest(TraceFile);
+
+	FFileOutDataStream DataStream;
+	if (!DataStream.Open(TraceFile))
+	{
+		return;
+	}
+
+	FTraceWriter Writer(DataStream);
+	Writer.Begin();
+	WriteDiagnosticsSession2Event(Writer);
+	FMiscWriter MiscWriter(Writer);
+	MiscWriter.Test();
+	Writer.End();
+
+	EndTest(Writer);
+}
+
+////////////////////////////////////////////////////////////////////////////////////////////////////
+
+void TraceWriterTest2b(const TCHAR* TraceFile)
+{
+	BeginTest(TraceFile);
+
+	FFileOutDataStream DataStream;
+	if (!DataStream.Open(TraceFile))
+	{
+		return;
+	}
+
+	FTraceWriter Writer(DataStream);
+	Writer.Begin();
+	WriteDiagnosticsSession2Event(Writer);
+	FMiscWriter MiscWriter(Writer);
+	MiscWriter.Test2();
+	Writer.End();
+
+	EndTest(Writer);
+}
+
+////////////////////////////////////////////////////////////////////////////////////////////////////
+
+static void WriteTestEvents(FTraceWriter& Writer, ETraceWriterEventFlags EventFlags)
+{
+	const uint32 GameThreadId = Writer.RegisterThread(ANSITEXTVIEW("GameThread"));
+
+	// Test if threads are cached.
+	const uint32 GameThreadId2 = Writer.RegisterThread(ANSITEXTVIEW("GameThread"));
+	check(GameThreadId2 == GameThreadId);
+
+	FAnsiStringView Logger("Test");
+
+	FAnsiString SimpleTypesEventName("SimpleTypes");
+	FAnsiString ArrayTypesEventName("ArrayTypes");
+
+	bool bIsImportant = false;
+	if (EnumHasAnyFlags(EventFlags, ETraceWriterEventFlags::Important))
+	{
+		bIsImportant = true;
+		SimpleTypesEventName += "_Important";
+		ArrayTypesEventName += "_Important";
+	}
+	if (EnumHasAnyFlags(EventFlags, ETraceWriterEventFlags::NoSync))
+	{
+		SimpleTypesEventName += "_NoSync";
+		ArrayTypesEventName += "_NoSync";
+	}
+	else
+	{
+		SimpleTypesEventName += "_Sync";
+		ArrayTypesEventName += "_Sync";
+	}
+
+	const uint32 EventId_Test_SimpleTypes = Writer.DeclareEvent(Logger, SimpleTypesEventName, EventFlags)
+		.Field(ANSITEXTVIEW("Bool"),    ETraceWriterFieldType::Bool) // Protocol7 doesn't make any difference between Bool and Uint8 :(
+		.Field(ANSITEXTVIEW("Uint8"),   ETraceWriterFieldType::Uint8)
+		.Field(ANSITEXTVIEW("Uint16"),  ETraceWriterFieldType::Uint16)
+		.Field(ANSITEXTVIEW("Uint32"),  ETraceWriterFieldType::Uint32)
+		.Field(ANSITEXTVIEW("Uint64"),  ETraceWriterFieldType::Uint64)
+		.Field(ANSITEXTVIEW("Int8"),    ETraceWriterFieldType::Int8)
+		.Field(ANSITEXTVIEW("Int16"),   ETraceWriterFieldType::Int16)
+		.Field(ANSITEXTVIEW("Int32"),   ETraceWriterFieldType::Int32)
+		.Field(ANSITEXTVIEW("Int64"),   ETraceWriterFieldType::Int64)
+		.Field(ANSITEXTVIEW("Float32"), ETraceWriterFieldType::Float32)
+		.Field(ANSITEXTVIEW("Float64"), ETraceWriterFieldType::Float64)
+		.Field(ANSITEXTVIEW("Pointer"), ETraceWriterFieldType::Pointer) // Protocol7 doesn't make any difference between Pointer and Uint64 :(
+		.End();
+
+	const uint32 EventId_Test_ArrayTypes = Writer.DeclareEvent(Logger, ArrayTypesEventName, EventFlags)
+		.Field(ANSITEXTVIEW("Uint8Array"),   ETraceWriterFieldType::ArrayFlag | ETraceWriterFieldType::Uint8)
+		.Field(ANSITEXTVIEW("Uint16Array"),  ETraceWriterFieldType::ArrayFlag | ETraceWriterFieldType::Uint16)
+		.Field(ANSITEXTVIEW("Uint32Array"),  ETraceWriterFieldType::ArrayFlag | ETraceWriterFieldType::Uint32)
+		.Field(ANSITEXTVIEW("Uint64Array"),  ETraceWriterFieldType::ArrayFlag | ETraceWriterFieldType::Uint64)
+		.Field(ANSITEXTVIEW("Int8Array"),    ETraceWriterFieldType::ArrayFlag | ETraceWriterFieldType::Int8)
+		.Field(ANSITEXTVIEW("Int16Array"),   ETraceWriterFieldType::ArrayFlag | ETraceWriterFieldType::Int16)
+		.Field(ANSITEXTVIEW("Int32Array"),   ETraceWriterFieldType::ArrayFlag | ETraceWriterFieldType::Int32)
+		.Field(ANSITEXTVIEW("Int64Array"),   ETraceWriterFieldType::ArrayFlag | ETraceWriterFieldType::Int64)
+		.Field(ANSITEXTVIEW("Float32Array"), ETraceWriterFieldType::ArrayFlag | ETraceWriterFieldType::Float32)
+		.Field(ANSITEXTVIEW("Float64Array"), ETraceWriterFieldType::ArrayFlag | ETraceWriterFieldType::Float64)
+		.Field(ANSITEXTVIEW("PointerArray"), ETraceWriterFieldType::ArrayFlag | ETraceWriterFieldType::Pointer)
+		.End();
+
+	if (bIsImportant)
+	{
+		Writer.SetCurrentThreadImportants();
+	}
+	else
+	{
+		Writer.SetCurrentThread(GameThreadId);
+	}
+
+	Writer.WriteEvent(EventId_Test_SimpleTypes)
+		.Field("Bool", false)
+		.Field("Uint8", (uint8)0x10)
+		.Field("Uint16", (uint16)0x3210)
+		.Field("Uint32", 0x76543210u)
+		.Field("Uint64", 0xFEDCBA9876543210ull)
+		.Field("Int8", (int8)0x10)
+		.Field("Int16", (int16)0x3210)
+		.Field("Int32", (int32)0x76543210)
+		.Field("Int64", (int64)0xFEDCBA9876543210)
+		.Field("Float32", 43.21f)
+		.Field("Float64", 12.34)
+		.Field("Pointer", &Writer)
+		.End();
+
+	Writer.WriteEvent(EventId_Test_SimpleTypes)
+		.Field("Bool", true)
+		.Field("Uint8", (uint8)0x01)
+		.Field("Uint16", (uint16)0x1032)
+		.Field("Uint32", 0x10765432u)
+		.Field("Uint64", 0x10FEDCBA98765432ull)
+		.Field("Int8", (int8)0x01)
+		.Field("Int16", (int16)0x1032)
+		.Field("Int32", (int32)0x10765432)
+		.Field("Int64", (int64)0x10FEDCBA98765432)
+		.Field("Float32", 34.12f)
+		.Field("Float64", 21.43)
+		.Field("Pointer", &Writer)
+		.End();
+
+	TArray<uint8>  Uint8Array  = { 0, 1, 10, 42, 0xFF };
+	TArray<uint16> Uint16Array = { 0, 1, 0x00FF, 0xFF00, 0xFFFF };
+	TArray<uint32> Uint32Array = { 0u, 1u, 0x76543210u, 0xFEDCBA98u, 0xFFFFFFFFu };
+	TArray<uint64> Uint64Array = { 0ull, 1ull, 0xFEDCBA9876543210ull, 0x0123456789ABCDEFull, ~0ull };
+
+	TArray<int8>  Int8Array  = { 0, 1, -1, 100, MAX_int8, -100, MIN_int8 };
+	TArray<int16> Int16Array = { 0, 1, -1, 10000, MAX_int16, -10000, MIN_int16 };
+	TArray<int32> Int32Array = { 0, 1, -1, 100000, 987654321, MAX_int32, -100000, -987654321, MIN_int32 };
+	TArray<int64> Int64Array = { 0, 1, -1, 1000000, MAX_int64, -1000000, MIN_int64 };
+
+	TArray<float>   FloatArray = { 0.0f, 0.1f, 1.0f, -0.1f, -1.0f, MAX_flt, MIN_flt };
+	TArray<double> DoubleArray = { 0.0, 0.1, 1.0, -0.1, -1.0, MAX_dbl, MIN_dbl };
+
+	TArray<void*> PointerArray = { &Writer, &Writer };
+
+	// Arrays specified as (const void* Data, uint32 DataSize) pairs.
+	Writer.WriteEvent(EventId_Test_ArrayTypes)
+		.Field("Uint8Array",   Uint8Array.GetData(),   Uint8Array.Num()   * sizeof(uint8))
+		.Field("Uint16Array",  Uint16Array.GetData(),  Uint16Array.Num()  * sizeof(uint16))
+		.Field("Uint32Array",  Uint32Array.GetData(),  Uint32Array.Num()  * sizeof(uint32))
+		.Field("Uint64Array",  Uint64Array.GetData(),  Uint64Array.Num()  * sizeof(uint64))
+		.Field("Int8Array",    Int8Array.GetData(),    Int8Array.Num()    * sizeof(int8))
+		.Field("Int16Array",   Int16Array.GetData(),   Int16Array.Num()   * sizeof(int16))
+		.Field("Int32Array",   Int32Array.GetData(),   Int32Array.Num()   * sizeof(int32))
+		.Field("Int64Array",   Int64Array.GetData(),   Int64Array.Num()   * sizeof(int64))
+		.Field("Float32Array", FloatArray.GetData(),   FloatArray.Num()   * sizeof(float))
+		.Field("Float64Array", DoubleArray.GetData(),  DoubleArray.Num()  * sizeof(double))
+		.Field("PointerArray", PointerArray.GetData(), PointerArray.Num() * sizeof(void*))
+		.End();
+
+	// Arrays specified as TConstArrayView<>.
+	Writer.WriteEvent(EventId_Test_ArrayTypes)
+		.Field("Uint8Array",   TConstArrayView<uint8>  (Uint8Array.GetData(),   Uint8Array.Num()))
+		.Field("Uint16Array",  TConstArrayView<uint16> (Uint16Array.GetData(),  Uint16Array.Num()))
+		.Field("Uint32Array",  TConstArrayView<uint32> (Uint32Array.GetData(),  Uint32Array.Num()))
+		.Field("Uint64Array",  TConstArrayView<uint64> (Uint64Array.GetData(),  Uint64Array.Num()))
+		.Field("Int8Array",    TConstArrayView<int8>   (Int8Array.GetData(),    Int8Array.Num()))
+		.Field("Int16Array",   TConstArrayView<int16>  (Int16Array.GetData(),   Int16Array.Num()))
+		.Field("Int32Array",   TConstArrayView<int32>  (Int32Array.GetData(),   Int32Array.Num()))
+		.Field("Int64Array",   TConstArrayView<int64>  (Int64Array.GetData(),   Int64Array.Num()))
+		.Field("Float32Array", TConstArrayView<float>  (FloatArray.GetData(),   FloatArray.Num()))
+		.Field("Float64Array", TConstArrayView<double> (DoubleArray.GetData(),  DoubleArray.Num()))
+		.Field("PointerArray", TConstArrayView<void*>  (PointerArray.GetData(), PointerArray.Num()))
+		.End();
+
+	// Arrays specified as TArrayView<>.
+	Writer.WriteEvent(EventId_Test_ArrayTypes)
+		.Field("Uint8Array",   TArrayView<uint8>  (Uint8Array.GetData(),   Uint8Array.Num()))
+		.Field("Uint16Array",  TArrayView<uint16> (Uint16Array.GetData(),  Uint16Array.Num()))
+		.Field("Uint32Array",  TArrayView<uint32> (Uint32Array.GetData(),  Uint32Array.Num()))
+		.Field("Uint64Array",  TArrayView<uint64> (Uint64Array.GetData(),  Uint64Array.Num()))
+		.Field("Int8Array",    TArrayView<int8>   (Int8Array.GetData(),    Int8Array.Num()))
+		.Field("Int16Array",   TArrayView<int16>  (Int16Array.GetData(),   Int16Array.Num()))
+		.Field("Int32Array",   TArrayView<int32>  (Int32Array.GetData(),   Int32Array.Num()))
+		.Field("Int64Array",   TArrayView<int64>  (Int64Array.GetData(),   Int64Array.Num()))
+		.Field("Float32Array", TArrayView<float>  (FloatArray.GetData(),   FloatArray.Num()))
+		.Field("Float64Array", TArrayView<double> (DoubleArray.GetData(),  DoubleArray.Num()))
+		.Field("PointerArray", TArrayView<void*>  (PointerArray.GetData(), PointerArray.Num()))
+		.End();
+}
+
+////////////////////////////////////////////////////////////////////////////////////////////////////
+
+static void WriteAdhocMacrosTest(FTraceWriter& Writer)
+{
+	const uint32 TestEvent1 =
+	UE_ADHOC_TRACE_BEGIN_DECLARE_EVENT(Writer, "Test", "TestEvent1_Important", ImportantNoSync)
+		UE_ADHOC_TRACE_DECLARE_FIELD(bool, "BoolField")
+		UE_ADHOC_TRACE_DECLARE_FIELD(uint8, "Uint8Field")
+		UE_ADHOC_TRACE_DECLARE_FIELD(uint64, "Uint64Field")
+		UE_ADHOC_TRACE_DECLARE_FIELD(void*, "PointerField")
+		UE_ADHOC_TRACE_DECLARE_FIELD(uint8[], "ByteArrayField")
+		UE_ADHOC_TRACE_DECLARE_FIELD(uint64[], "Uint64ArrayField")
+		UE_ADHOC_TRACE_DECLARE_FIELD(void*[], "PointerArrayField")
+	UE_ADHOC_TRACE_END_DECLARE_EVENT
+
+	const uint32 TestEvent2 =
+	UE_ADHOC_TRACE_BEGIN_DECLARE_EVENT(Writer, "Test", "TestEvent2_Sync", None) // Sync
+		UE_ADHOC_TRACE_DECLARE_FIELD(bool, "BoolField")
+		UE_ADHOC_TRACE_DECLARE_FIELD(uint8, "Uint8Field")
+		UE_ADHOC_TRACE_DECLARE_FIELD(uint64, "Uint64Field")
+		UE_ADHOC_TRACE_DECLARE_FIELD(void*, "PointerField")
+		UE_ADHOC_TRACE_DECLARE_FIELD(uint8[], "ByteArrayField")
+		UE_ADHOC_TRACE_DECLARE_FIELD(uint64[], "Uint64ArrayField")
+		UE_ADHOC_TRACE_DECLARE_FIELD(void*[], "PointerArrayField")
+	UE_ADHOC_TRACE_END_DECLARE_EVENT
+
+	const uint32 TestEvent3 =
+	UE_ADHOC_TRACE_BEGIN_DECLARE_EVENT(Writer, "Test", "TestEvent3_NoSync", NoSync)
+		UE_ADHOC_TRACE_DECLARE_FIELD(bool, "BoolField")
+		UE_ADHOC_TRACE_DECLARE_FIELD(uint8, "Uint8Field")
+		UE_ADHOC_TRACE_DECLARE_FIELD(uint64, "Uint64Field")
+		UE_ADHOC_TRACE_DECLARE_FIELD(const void*, "PointerField")
+		UE_ADHOC_TRACE_DECLARE_FIELD(uint8[], "ByteArrayField")
+		UE_ADHOC_TRACE_DECLARE_FIELD(uint64[], "Uint64ArrayField")
+		UE_ADHOC_TRACE_DECLARE_FIELD(void*[], "PointerArrayField")
+	UE_ADHOC_TRACE_END_DECLARE_EVENT
+
+	TArray<uint8>   Uint8Array = { 0, 1, 10, 42, 0xFF };
+	TArray<uint64> Uint64Array = { 0ull, 1ull, 0xFEDCBA9876543210ull, 0x0123456789ABCDEFull, ~0ull };
+
+	TArray<void*> PointerArray = { &Writer, &Writer };
+
+	Writer.SetCurrentThreadImportants();
+
+	UE_ADHOC_TRACE_LOG(Writer, TestEvent1,
+		(bool)true,
+		(uint8)42,
+		0x0123456789ABCDEFull,
+		&Writer,
+		TArrayView<uint8>(Uint8Array.GetData(), Uint8Array.Num()),
+		TArrayView<uint64>(Uint64Array.GetData(), Uint64Array.Num()),
+		TArrayView<void*>(PointerArray.GetData(), PointerArray.Num()))
+
+	const uint32 GameThreadId = Writer.RegisterThread(ANSITEXTVIEW("GameThread"));
+
+	Writer.SetCurrentThread(GameThreadId);
+
+	UE_ADHOC_TRACE_LOG(Writer, TestEvent2,
+		(bool)true,
+		(uint8)42,
+		0x0123456789ABCDEFull,
+		&Writer,
+		TArrayView<uint8>(Uint8Array.GetData(), Uint8Array.Num()),
+		TArrayView<uint64>(Uint64Array.GetData(), Uint64Array.Num()),
+		TArrayView<void*>(PointerArray.GetData(), PointerArray.Num()))
+
+	UE_ADHOC_TRACE_LOG(Writer, TestEvent3,
+		(bool)true,
+		(uint8)42,
+		0x0123456789ABCDEFull,
+		&Writer,
+		TConstArrayView<uint8>(Uint8Array.GetData(), Uint8Array.Num()),
+		TConstArrayView<uint64>(Uint64Array.GetData(), Uint64Array.Num()),
+		TConstArrayView<void*>(PointerArray.GetData(), PointerArray.Num()))
+}
+
+////////////////////////////////////////////////////////////////////////////////////////////////////
+
+static void WriteScopedTestEvents(FTraceWriter& Writer)
+{
+	const uint32 TestScopedEvent1 =
+	UE_ADHOC_TRACE_BEGIN_DECLARE_EVENT(Writer, "Test", "ScopedEvent1_NoSync", NoSync)
+		UE_ADHOC_TRACE_DECLARE_FIELD(uint32, "Id")
+		UE_ADHOC_TRACE_DECLARE_FIELD(AnsiString, "Name")
+	UE_ADHOC_TRACE_END_DECLARE_EVENT
+
+	const uint32 TestScopedEvent2 =
+	UE_ADHOC_TRACE_BEGIN_DECLARE_EVENT(Writer, "Test", "ScopedEvent2_Sync", None) // Sync
+		UE_ADHOC_TRACE_DECLARE_FIELD(uint32, "Id")
+		UE_ADHOC_TRACE_DECLARE_FIELD(AnsiString, "Name")
+	UE_ADHOC_TRACE_END_DECLARE_EVENT
+
+	const uint32 GameThreadId = Writer.RegisterThread(ANSITEXTVIEW("GameThread"));
+	Writer.SetCurrentThread(GameThreadId);
+
+	Writer.WriteEnterScopeEvent();
+	Writer.WriteEvent(TestScopedEvent1)
+		.Field("Id", 1)
+		.Field("Name", "abc")
+		.End();
+	//...
+	Writer.WriteLeaveScopeEvent();
+
+	Writer.WriteEnterScopeEvent();
+	Writer.WriteEvent(TestScopedEvent2)
+		.Field("Id", 2)
+		.Field("Name", "xyz")
+		.End();
+	//...
+	Writer.WriteLeaveScopeEvent();
+
+	Writer.WriteStampedEnterScopeEvent(Writer.GetTime());
+	Writer.WriteEvent(TestScopedEvent1)
+		.Field("Id", 3)
+		.Field("Name", "abc2")
+		.End();
+	//...
+	Writer.WriteStampedLeaveScopeEvent(Writer.GetTime());
+
+	Writer.WriteStampedEnterScopeEvent(Writer.GetTime());
+	Writer.WriteEvent(TestScopedEvent2)
+		.Field("Id", 4)
+		.Field("Name", "xyz2")
+		.End();
+	//...
+	Writer.WriteStampedLeaveScopeEvent(Writer.GetTime());
+}
+
+////////////////////////////////////////////////////////////////////////////////////////////////////
+
+void TraceWriterTest3(const TCHAR* TraceFile)
+{
+	BeginTest(TraceFile);
+
+	FFileOutDataStream DataStream;
+	if (!DataStream.Open(TraceFile))
+	{
+		return;
+	}
+
+	FTraceWriter Writer(DataStream);
+	Writer.Begin();
+	WriteDiagnosticsSession2Event(Writer);
+	WriteTestEvents(Writer, ETraceWriterEventFlags::None); // Sync
+	WriteTestEvents(Writer, ETraceWriterEventFlags::ImportantNoSync);
+	WriteTestEvents(Writer, ETraceWriterEventFlags::NoSync);
+	WriteAdhocMacrosTest(Writer);
+	WriteScopedTestEvents(Writer);
+	Writer.End();
+
+	EndTest(Writer);
+}
+
+////////////////////////////////////////////////////////////////////////////////////////////////////
+// FCpuProfilerEventBatchWriter
+////////////////////////////////////////////////////////////////////////////////////////////////////
+
+class FCpuProfilerEventBatchWriter
+{
+public:
+	FCpuProfilerEventBatchWriter()
+	{
+	}
+	virtual ~FCpuProfilerEventBatchWriter()
+	{
+		if (Buffer)
+		{
+			FMemory::Free(Buffer);
+		}
+	}
+
+	void Begin(uint32 MinInitialSize = 8 * 1024)
+	{
+		UsedSize = 0;
+		if (MinInitialSize > AllocatedSize)
+		{
+			Resize(MinInitialSize);
+		}
+		LastCpuEventTimestamp = 0;
+	}
+
+	void BeginCpuEvent(uint32 SpecId, uint64 Timestamp);
+	void EndCpuEvent(uint64 Timestamp);
+
+	TConstArrayView<uint8> End()
+	{
+		return TConstArrayView<uint8>(Buffer, UsedSize);
+	}
+
+	void Resize(uint32 NewSize)
+	{
+		if (NewSize == AllocatedSize)
+		{
+			return;
+		}
+		uint8* NewBuffer = nullptr;
+		if (NewSize > 0)
+		{
+			NewBuffer = (uint8*)FMemory::Malloc(NewSize);
+			if (UsedSize > 0)
+			{
+				FMemory::Memcpy(NewBuffer, Buffer, FMath::Min(UsedSize, NewSize));
+			}
+		}
+		if (Buffer)
+		{
+			FMemory::Free(Buffer);
+		}
+		Buffer = NewBuffer;
+		AllocatedSize = NewSize;
+		if (UsedSize > AllocatedSize)
+		{
+			UsedSize = AllocatedSize;
+		}
+	}
+
+	uint8* GetBuffer(uint32 Required)
+	{
+		constexpr uint32 PageSize = 4 * 1024;
+		if (UsedSize + Required > AllocatedSize)
+		{
+			Resize((UsedSize + Required + PageSize - 1) & ~(PageSize - 1));
+		}
+		return Buffer + UsedSize;
+	}
+
+private:
+	uint8* Buffer = nullptr;
+	uint32 AllocatedSize = 0;
+	uint32 UsedSize = 0;
+	uint64 LastCpuEventTimestamp = 0;
+};
+
+////////////////////////////////////////////////////////////////////////////////////////////////////
+
+void FCpuProfilerEventBatchWriter::BeginCpuEvent(uint32 InSpecId, uint64 InTimestamp)
+{
+	uint8* StartCursor = GetBuffer(10 + 5);
+	uint8* Cursor = StartCursor;
+
+	uint64 CycleDiff = InTimestamp - LastCpuEventTimestamp;
+	LastCpuEventTimestamp = InTimestamp;
+	
+	// Write Timestamp
+	FTraceUtils::Encode7bit((CycleDiff << 2) | 1ull, Cursor);
+	
+	// Write SpecId
+	FTraceUtils::Encode7bit((InSpecId << 1) | 0, Cursor);
+
+	UsedSize += (Cursor - StartCursor);
+}
+
+////////////////////////////////////////////////////////////////////////////////////////////////////
+
+void FCpuProfilerEventBatchWriter::EndCpuEvent(uint64 Timestamp)
+{
+	uint8* StartCursor = GetBuffer(10);
+	uint8* Cursor = StartCursor;
+
+	uint64 CycleDiff = Timestamp - LastCpuEventTimestamp;
+	LastCpuEventTimestamp = Timestamp;
+
+	// Write Timestamp
+	FTraceUtils::Encode7bit(CycleDiff << 2, Cursor);
+
+	UsedSize += (Cursor - StartCursor);
+}
+
+////////////////////////////////////////////////////////////////////////////////////////////////////
+// FCpuProfilerWriter
+////////////////////////////////////////////////////////////////////////////////////////////////////
+
+class FCpuProfilerWriter
+{
+public:
+	FCpuProfilerWriter(FTraceWriter& InWriter)
+		: Writer(InWriter)
+	{
+	}
+	virtual ~FCpuProfilerWriter()
+	{
+	}
+
+	void DeclareEvents();
+
+	void WriteEventSpecEvent(uint32 SpecId, FAnsiStringView Name, FAnsiStringView File = FAnsiStringView(), uint32 Line = 0);
+	void WriteEventBatchEvent(TConstArrayView<uint8> Data);
+
+	void Test();
+
+private:
+	FTraceWriter& Writer;
+	uint32 EventId_EventSpec = 0;
+	uint32 EventId_MetadataSpec = 0;
+	uint32 EventId_Metadata = 0;
+	uint32 EventId_EventBatchV3 = 0;
+	uint32 EventId_EndThread = 0;
+};
+
+////////////////////////////////////////////////////////////////////////////////////////////////////
+
+void FCpuProfilerWriter::DeclareEvents()
+{
+	if (EventId_EventSpec != 0)
+	{
+		return;
+	}
+
+	EventId_EventSpec =
+	UE_ADHOC_TRACE_BEGIN_DECLARE_EVENT(Writer, "CpuProfiler", "EventSpec", ImportantNoSync)
+		UE_ADHOC_TRACE_DECLARE_FIELD(uint32, "Id")
+		UE_ADHOC_TRACE_DECLARE_FIELD(UE::Trace::AnsiString, "Name")
+		UE_ADHOC_TRACE_DECLARE_FIELD(UE::Trace::AnsiString, "File")
+		UE_ADHOC_TRACE_DECLARE_FIELD(uint32, "Line")
+	UE_ADHOC_TRACE_END_DECLARE_EVENT
+
+	EventId_MetadataSpec =
+	UE_ADHOC_TRACE_BEGIN_DECLARE_EVENT(Writer, "CpuProfiler", "MetadataSpec", ImportantNoSync)
+		UE_ADHOC_TRACE_DECLARE_FIELD(uint32, "Id")
+		UE_ADHOC_TRACE_DECLARE_FIELD(UE::Trace::AnsiString, "Name")
+		UE_ADHOC_TRACE_DECLARE_FIELD(UE::Trace::WideString, "NameFormat")
+		UE_ADHOC_TRACE_DECLARE_FIELD(uint8[], "FieldNames")
+	UE_ADHOC_TRACE_END_DECLARE_EVENT
+
+	EventId_Metadata =
+	UE_ADHOC_TRACE_BEGIN_DECLARE_EVENT(Writer, "CpuProfiler", "Metadata", NoSync)
+		UE_ADHOC_TRACE_DECLARE_FIELD(uint32, "Id")
+		UE_ADHOC_TRACE_DECLARE_FIELD(uint32, "SpecId")
+		UE_ADHOC_TRACE_DECLARE_FIELD(uint8[], "Metadata")
+	UE_ADHOC_TRACE_END_DECLARE_EVENT
+
+	EventId_EventBatchV3 =
+	UE_ADHOC_TRACE_BEGIN_DECLARE_EVENT(Writer, "CpuProfiler", "EventBatchV3", NoSync)
+		UE_ADHOC_TRACE_DECLARE_FIELD(uint8[], "Data")
+	UE_ADHOC_TRACE_END_DECLARE_EVENT
+
+	EventId_EndThread =
+	UE_ADHOC_TRACE_BEGIN_DECLARE_EVENT(Writer, "CpuProfiler", "EndThread", NoSync)
+		UE_ADHOC_TRACE_DECLARE_FIELD(uint64, "Cycle") // added in UE 5.4
+	UE_ADHOC_TRACE_END_DECLARE_EVENT
+}
+
+////////////////////////////////////////////////////////////////////////////////////////////////////
+
+void FCpuProfilerWriter::WriteEventSpecEvent(uint32 SpecId, FAnsiStringView Name, FAnsiStringView File, uint32 Line)
+{
+	Writer.WriteEvent(EventId_EventSpec)
+		.Field(0, SpecId)
+		.Field(1, Name)
+		.Field(2, File)
+		.Field(3, Line)
+		.End();
+}
+
+////////////////////////////////////////////////////////////////////////////////////////////////////
+
+void FCpuProfilerWriter::WriteEventBatchEvent(TConstArrayView<uint8> Data)
+{
+	Writer.WriteEvent(EventId_EventBatchV3)
+		.Field(0, Data)
+		.End();
+}
+
+////////////////////////////////////////////////////////////////////////////////////////////////////
+
+void FCpuProfilerWriter::Test()
+{
+	DeclareEvents();
+}
+
+////////////////////////////////////////////////////////////////////////////////////////////////////
+// FTimingEventsWriter
+////////////////////////////////////////////////////////////////////////////////////////////////////
+
+class FTimingEventsWriter
+{
+public:
+	FTimingEventsWriter(FTraceWriter& InWriter, FCpuProfilerWriter& InCpuProfilerWriter, FCpuProfilerEventBatchWriter& InBatchWriter, uint32 InNumSpecs)
+		: Writer(InWriter)
+		, CpuProfilerWriter(InCpuProfilerWriter)
+		, BatchWriter(InBatchWriter)
+		, NumSpecs(InNumSpecs)
+	{
+		const uint64 OneSecond = Writer.GetTimeFrequency();
+		MinParentEventDuration = OneSecond / 1000000; // 1us
+	}
+	~FTimingEventsWriter()
+	{
+	}
+
+	void BeginTimingEvent(uint64 StartTime);
+	void EndTimingEvent(uint64 EndTime);
+	void GenerateTimingEvents(uint32 Depth, uint64 IntervalStartTime, uint64 IntervalDuration);
+
+private:
+	FTraceWriter& Writer;
+	FCpuProfilerWriter& CpuProfilerWriter;
+	FCpuProfilerEventBatchWriter& BatchWriter;
+	uint32 TotalTimingEvents = 0;
+	static constexpr uint32 FlushAfterNumTimingEvents = 200;
+	uint64 MinParentEventDuration = 0;
+	uint32 NumSpecs = 0;
+	static constexpr uint32 MaxDepth = 16;
+	static constexpr uint32 MaxChildTimingEvents = 10;
+};
+
+////////////////////////////////////////////////////////////////////////////////////////////////////
+
+void FTimingEventsWriter::BeginTimingEvent(uint64 StartTime)
+{
+	++TotalTimingEvents;
+	if (TotalTimingEvents % FlushAfterNumTimingEvents == 0)
+	{
+		CpuProfilerWriter.WriteEventBatchEvent(BatchWriter.End());
+		BatchWriter.Begin();
+	}
+
+	uint32 SpecId = 1 + FMath::RandHelper(NumSpecs);
+	check(SpecId >= 1 && SpecId <= NumSpecs);
+	BatchWriter.BeginCpuEvent(SpecId, StartTime);
+}
+
+////////////////////////////////////////////////////////////////////////////////////////////////////
+
+void FTimingEventsWriter::EndTimingEvent(uint64 EndTime)
+{
+	BatchWriter.EndCpuEvent(EndTime);
+}
+
+////////////////////////////////////////////////////////////////////////////////////////////////////
+
+void FTimingEventsWriter::GenerateTimingEvents(uint32 Depth, uint64 IntervalStartTime, uint64 IntervalDuration)
+{
+	if (IntervalDuration == 0)
+	{
+		return;
+	}
+
+	const uint64 IntervalEndTime = IntervalStartTime + IntervalDuration;
+
+	uint64 TimingEventTime = IntervalStartTime + uint64(FMath::RandHelper64((int64)IntervalDuration /10));
+
+	while (true)
+	{
+		const uint64 TimingEventDuration = uint64(FMath::RandHelper64((int64)IntervalDuration));
+		if (TimingEventTime + TimingEventDuration > IntervalEndTime)
+		{
+			break;
+		}
+
+		BeginTimingEvent(TimingEventTime);
+
+		if (Depth < MaxDepth && TimingEventDuration > MinParentEventDuration)
+		{
+			GenerateTimingEvents(Depth + 1, TimingEventTime, TimingEventDuration);
+		}
+
+		TimingEventTime += TimingEventDuration;
+		EndTimingEvent(TimingEventTime);
+	}
+}
+
+////////////////////////////////////////////////////////////////////////////////////////////////////
+////////////////////////////////////////////////////////////////////////////////////////////////////
+
+static void WriteCpuProfilerTest(FTraceWriter& Writer)
+{
+	FCpuProfilerWriter CpuProfilerWriter(Writer);
+	FMiscWriter MiscWriter(Writer);
+
+	const uint32 GameThreadId = Writer.RegisterThread(ANSITEXTVIEW("GameThread"));
+	const uint32 RenderThreadId = Writer.RegisterThread(ANSITEXTVIEW("RenderThread"));
+
+	CpuProfilerWriter.DeclareEvents();
+	MiscWriter.DeclareFrameEvents();
+
+	Writer.SetCurrentThreadImportants();
+
+	constexpr uint32 NumSpecs = 26;
+	for (uint32 SpecId = 1; SpecId <= NumSpecs; ++SpecId)
+	{
+		ANSICHAR Name[10];
+		Name[0] = 'A' + SpecId - 1;
+		Name[1] = 0;
+		CpuProfilerWriter.WriteEventSpecEvent(SpecId, Name);
+	}
+
+	const uint64 BaseTimestamp = Writer.GetStartTime();
+	const uint64 OneSecond = Writer.GetTimeFrequency();
+	const uint64 OneMilliSecond = OneSecond / 1000;
+	const uint64 OneMicroSecond = OneSecond / 1000000;
+
+	uint64 MinFrameDuration = (10 * OneSecond) / 1000; // 10ms
+	uint64 MaxFrameDuration = (20 * OneSecond) / 1000; // 20ms
+	uint64 MaxDurationBetweenFrames = OneSecond / 1000; // 1ms
+
+	FCpuProfilerEventBatchWriter BatchWriter;
+	FTimingEventsWriter TimingEventsWriter(Writer, CpuProfilerWriter, BatchWriter, NumSpecs);
+
+	auto GenerateFrames =
+		[&]
+		(uint64 StartTime, uint32 FrameCount, uint8 FrameType) -> uint64
+		{
+			uint64 CurrentTime = StartTime;
+			uint32 NumTimingEvents = 0;
+
+			for (uint32 FrameIndex = 0; FrameIndex < FrameCount; ++FrameIndex)
+			{
+				MiscWriter.WriteBeginFrameEvent(CurrentTime, FrameType);
+
+				const uint64 FrameDuration = MinFrameDuration + FMath::RandHelper64(int64(MaxFrameDuration - MinFrameDuration));
+
+				BatchWriter.Begin();
+				TimingEventsWriter.GenerateTimingEvents(0, CurrentTime, FrameDuration);
+				CpuProfilerWriter.WriteEventBatchEvent(BatchWriter.End());
+
+				CurrentTime += FrameDuration;
+				MiscWriter.WriteEndFrameEvent(CurrentTime, FrameType);
+
+				CurrentTime += uint64(FMath::RandHelper64(MaxDurationBetweenFrames));
+			}
+
+			return CurrentTime;
+		};
+
+	FMath::RandInit(42);
+
+	constexpr uint64 NumPackets = 10;
+	constexpr uint64 FramesPerPacket = 100;
+
+	uint64 CurrentTime = BaseTimestamp;
+
+	for (uint64 PacketIndex = 0; PacketIndex < NumPackets; ++PacketIndex)
+	{
+		Writer.SetCurrentThread(GameThreadId);
+		uint64 CurrentTimeA = GenerateFrames(CurrentTime, FramesPerPacket, (uint8)TraceFrameType_Game);
+
+		Writer.SetCurrentThread(RenderThreadId);
+		uint64 CurrentTimeB = GenerateFrames(CurrentTime, FramesPerPacket, (uint8)TraceFrameType_Rendering);
+
+		CurrentTime = FMath::Max(CurrentTimeA, CurrentTimeB);
+		CurrentTime += uint64(FMath::RandHelper64(MaxDurationBetweenFrames));
+	}
+}
+
+////////////////////////////////////////////////////////////////////////////////////////////////////
+
+void TraceWriterTestCpuProfiler(const TCHAR* TraceFile)
+{
+	BeginTest(TraceFile);
+
+	FFileOutDataStream DataStream;
+	if (!DataStream.Open(TraceFile))
+	{
+		return;
+	}
+
+	FTraceWriter Writer(DataStream);
+	Writer.Begin();
+	WriteDiagnosticsSession2Event(Writer);
+	WriteCpuProfilerTest(Writer);
+	Writer.End();
+
+	EndTest(Writer);
+}
+
+////////////////////////////////////////////////////////////////////////////////////////////////////
+
+} // namespace UE::Trace::Tests
+
+UE_ENABLE_OPTIMIZATION
diff --git a/Engine/Source/Programs/TraceTrimmer/Private/TraceWriterTests.h b/Engine/Source/Programs/TraceTrimmer/Private/TraceWriterTests.h
new file mode 100644
--- /dev/null
+++ b/Engine/Source/Programs/TraceTrimmer/Private/TraceWriterTests.h
@@ -0,0 +1,17 @@
+// Copyright Epic Games, Inc. All Rights Reserved.
+
+#pragma once
+
+#include "CoreTypes.h"
+
+namespace UE::Trace::Tests
+{
+
+void TraceWriterTest0(const TCHAR* TraceFile);
+void TraceWriterTest1(const TCHAR* TraceFile);
+void TraceWriterTest2a(const TCHAR* TraceFile);
+void TraceWriterTest2b(const TCHAR* TraceFile);
+void TraceWriterTest3(const TCHAR* TraceFile);
+void TraceWriterTestCpuProfiler(const TCHAR* TraceFile);
+
+} // namespace UE::Trace::Tests
diff --git a/Engine/Source/Programs/TraceTrimmer/TraceTrimmer.Build.cs b/Engine/Source/Programs/TraceTrimmer/TraceTrimmer.Build.cs
new file mode 100644
--- /dev/null
+++ b/Engine/Source/Programs/TraceTrimmer/TraceTrimmer.Build.cs
@@ -0,0 +1,15 @@
+// Copyright Epic Games, Inc. All Rights Reserved.
+
+using UnrealBuildTool;
+
+public class TraceTrimmer : ModuleRules
+{
+	public TraceTrimmer(ReadOnlyTargetRules Target) : base(Target)
+	{
+		PublicIncludePathModuleNames.Add("Launch");
+		PrivateDependencyModuleNames.Add("Core");
+		PrivateDependencyModuleNames.Add("Projects");
+		PrivateDependencyModuleNames.Add("TraceAnalysis");
+		PrivateDependencyModuleNames.Add("TraceServices");
+	}
+}
diff --git a/Engine/Source/Programs/TraceTrimmer/TraceTrimmer.Target.cs b/Engine/Source/Programs/TraceTrimmer/TraceTrimmer.Target.cs
new file mode 100644
--- /dev/null
+++ b/Engine/Source/Programs/TraceTrimmer/TraceTrimmer.Target.cs
@@ -0,0 +1,58 @@
+// Copyright Epic Games, Inc. All Rights Reserved.
+
+using UnrealBuildTool;
+using System.Collections.Generic;
+
+[SupportedPlatforms(UnrealPlatformClass.All)]
+public class TraceTrimmerTarget : TargetRules
+{
+	public TraceTrimmerTarget(TargetInfo Target) : base(Target)
+	{
+		Type = TargetType.Program;
+		LaunchModuleName = "TraceTrimmer";
+		LinkType = TargetLinkType.Monolithic;
+
+		IncludeOrderVersion = EngineIncludeOrderVersion.Latest;
+
+		// This app is a console application, not a Windows app (sets entry point to main(), instead of WinMain())
+		bIsBuildingConsoleApplication = true;
+
+		// Lean and mean
+		bBuildDeveloperTools = false;
+
+		bBuildWithEditorOnlyData = false;
+
+		//bUseStaticCRT = true;
+
+		// Currently this app is not linking against the engine, so we'll compile out references from Core to the rest of the engine
+		bCompileAgainstEngine = false;
+		bCompileAgainstCoreUObject = false;
+		bCompileAgainstApplicationCore = false;
+		bCompileICU = false;
+		//bCompileCEF3 = false;
+
+		bEnableTrace = false; // UE_TRACE_ENABLED
+		GlobalDefinitions.Add("CPUPROFILERTRACE_ENABLED=0");
+		GlobalDefinitions.Add("GPUPROFILERTRACE_ENABLED=0");
+		GlobalDefinitions.Add("LOADTIMEPROFILERTRACE_ENABLED=0");
+		GlobalDefinitions.Add("STATSTRACE_ENABLED=0");
+		GlobalDefinitions.Add("COUNTERSTRACE_ENABLED=0");
+		GlobalDefinitions.Add("LOGTRACE_ENABLED=0");
+		GlobalDefinitions.Add("MISCTRACE_ENABLED=0");
+		GlobalDefinitions.Add("UE_TASK_TRACE_ENABLED=0");
+		GlobalDefinitions.Add("UE_NET_TRACE_ENABLED=0");
+		GlobalDefinitions.Add("UE_MEMORY_TAGS_TRACE_ENABLED=0");
+		GlobalDefinitions.Add("UE_MEMORY_TRACE_ENABLED=0");
+		GlobalDefinitions.Add("UE_CALLSTACK_TRACE_ENABLED=0");
+
+		// see Engine\Source\Developer\TraceAnalysis\Public\TraceAnalysisDebug.h
+		GlobalDefinitions.Add("UE_TRACE_ANALYSIS_DEBUG_API=0");
+		GlobalDefinitions.Add("UE_TRACE_ANALYSIS_DEBUG=0");
+		GlobalDefinitions.Add("UE_TRACE_ANALYSIS_DEBUG_LOG_IMPL=1");
+		GlobalDefinitions.Add("UE_TRACE_ANALYSIS_DEBUG_LEVEL=2");
+
+		//GlobalDefinitions.Add("_CRT_SECURE_NO_WARNINGS");
+		//GlobalDefinitions.Add("_CRT_NONSTDC_NO_WARNINGS");
+		//WindowsPlatform.TargetWindowsVersion = 0x0603; // Windows 8.1
+	}
+}
