[TraceTrimmer]
- Added support for optional -i=... parameter (for specifying the input file).
- Added support for default output file (i.e. if -o=... is not specified); it default to <in_trace_file>.out.utrace.
- Removed -tracehost=, -traceid= and -outTracehost= from usage help text (as support for these is not yet implemented).

#Synced-CL 50651828

diff --git a/Engine/Source/Programs/TraceTrimmer/Private/TraceTrimmer.cpp b/Engine/Source/Programs/TraceTrimmer/Private/TraceTrimmer.cpp
--- a/Engine/Source/Programs/TraceTrimmer/Private/TraceTrimmer.cpp
+++ b/Engine/Source/Programs/TraceTrimmer/Private/TraceTrimmer.cpp
@@ -38,17 +38,18 @@ namespace UE::TraceTrimmer
 static void PrintUsage()
 {
 	puts("Usage:");
-	puts("  TraceTrimmer [<in_trace_file> | -tracehost=... -traceid=...] [-o=... | -outTraceHost=...] [<options...>]");
+	//puts("  TraceTrimmer [-i=][<in_trace_file>] [-tracehost=... -traceid=...] [-o=...] [<options...>]");
+	puts("  TraceTrimmer [-i=][<in_trace_file>] [-o=...] [<options...>]");
 	puts("Where:");
-	puts("  <in_trace_file> : the input *.utrace file");
-	puts("  -tracehost=<host:port> : the host address of the input trace server");
-	puts("  -traceid=<id> : the id of the input trace (identified by the trace server)");
-	puts("  -o=<out_trace_file> : the output trace file; defaults to <in_trace_file>.out.utrace");
-	puts("  -outTraceHost=<out_host> : the host address of the output trace server");
-	puts("  -startTime=<time> : start time filter in seconds");
-	puts("  -endTime=<time> : end time filter in seconds");
-	puts("  -include=<events> : comma separated list of trace events to include");
-	puts("  -exclude=<events> : comma separated list of trace events to exclude");
+	puts("  -i=<in_trace_file> : the input *.utrace file");
+	//puts("  -tracehost=<host:port> : the host address of the input trace server");
+	//puts("  -traceid=<id> : the id of the input trace (identified by the trace server)");
+	puts("  -o=<out_trace_file> : the output *.utrace file; defaults to <in_trace_file>.out.utrace");
+	//puts("  -outTraceHost=<out_host> : the host address of the output trace server");
+	puts("  -startTime=<time> : the start time filter, in seconds");
+	puts("  -endTime=<time> : the end time filter, in seconds");
+	puts("  -include=<events> : a comma separated list of trace events to include");
+	puts("  -exclude=<events> : a comma separated list of trace events to exclude");
 	puts("Note:");
 	puts("    The include/exclude events are specified as \"logger.event\" wildcard patterns, ex.: \"$trace.*,cpu.*\".");
 #if TRACE_TRIMMER_HAS_TRACE_WRITER_TESTS
@@ -126,7 +127,22 @@ static int32 ParseCommandLine(FCmdLineParams& Parameters, int32 ArgC, TCHAR* Arg
 	{
 		const TCHAR* Arg = ArgV[ArgIndex];
 
-		if (FCString::Strnicmp(TEXT("-tracehost="), Arg, 11) == 0)
+		if (FCString::Strnicmp(TEXT("-i="), Arg, 3) == 0)
+		{
+			if (!Parameters.InputTraceFile.IsEmpty())
+			{
+				LogTraceTrimmerError("The input trace file cannot be specified multiple times!");
+				return -1;
+			}
+			if (!Parameters.InputTraceServer.IsEmpty())
+			{
+				LogTraceTrimmerError("Cannot specify both an input trace file and an input trace server!");
+				return -1;
+			}
+			Parameters.InputTraceFile = Arg + 3;
+		}
+#if 0 // not yet supported
+		else if (FCString::Strnicmp(TEXT("-tracehost="), Arg, 11) == 0)
 		{
 			if (!Parameters.InputTraceServer.IsEmpty())
 			{
@@ -161,6 +177,7 @@ static int32 ParseCommandLine(FCmdLineParams& Parameters, int32 ArgC, TCHAR* Arg
 				Parameters.InputTraceId = FCString::Atoi(Arg + 9);
 			}
 		}
+#endif
 		else if (FCString::Strnicmp(TEXT("-o="), Arg, 3) == 0)
 		{
 			if (!Parameters.OutputTraceFile.IsEmpty())
@@ -175,6 +192,7 @@ static int32 ParseCommandLine(FCmdLineParams& Parameters, int32 ArgC, TCHAR* Arg
 			}
 			Parameters.OutputTraceFile = Arg + 3;
 		}
+#if 0 // not yet supported
 		else if (FCString::Strnicmp(TEXT("-outTraceHost="), Arg, 14) == 0)
 		{
 			if (!Parameters.OutputTraceServer.IsEmpty())
@@ -189,6 +207,7 @@ static int32 ParseCommandLine(FCmdLineParams& Parameters, int32 ArgC, TCHAR* Arg
 			}
 			Parameters.OutputTraceServer = Arg + 14;
 		}
+#endif
 		else if (FCString::Strnicmp(TEXT("-startTime="), Arg, 11) == 0)
 		{
 			Parameters.TrimParams.StartTime = FCString::Atof(Arg + 11);
@@ -211,6 +230,27 @@ static int32 ParseCommandLine(FCmdLineParams& Parameters, int32 ArgC, TCHAR* Arg
 		}
 	}
 
+	if (Parameters.InputTraceFile.IsEmpty() && Parameters.InputTraceServer.IsEmpty())
+	{
+		LogTraceTrimmerError("The input trace is not specified!");
+		return -1;
+	}
+
+	if (Parameters.OutputTraceFile.IsEmpty() && Parameters.OutputTraceServer.IsEmpty())
+	{
+		if (!Parameters.InputTraceFile.IsEmpty())
+		{
+			Parameters.OutputTraceFile = FPaths::GetBaseFilename(Parameters.InputTraceFile, false) + TEXT(".out.utrace");
+			FPaths::NormalizeFilename(Parameters.OutputTraceFile);
+			LogTraceTrimmerWarning("The output trace is not specified. Will default to: \"%s\"", *Parameters.OutputTraceFile);
+		}
+		else
+		{
+			LogTraceTrimmerError("The input trace is not a file and the output trace file is not specified!");
+			return -1;
+		}
+	}
+
 	return 1;
 }
 
