[TraceTrimmer] Fixed crash on shutdown when debugger is attached. Added RequestEngineExit() in main before closing app to ensure FUObjectArray::RemoveUObjectDeleteListener() is not called from FSparseDelegateStorage::FObjectListener::~FObjectListener() (becsaue it would try to lock GUObjectArray's already-destroyed mutex).

#Synced-CL 52393202

diff --git a/Engine/Source/Programs/TraceTrimmer/Private/TraceTrimmer.cpp b/Engine/Source/Programs/TraceTrimmer/Private/TraceTrimmer.cpp
--- a/Engine/Source/Programs/TraceTrimmer/Private/TraceTrimmer.cpp
+++ b/Engine/Source/Programs/TraceTrimmer/Private/TraceTrimmer.cpp
@@ -77,7 +77,7 @@ struct FCmdLineParams
 ////////////////////////////////////////////////////////////////////////////////////////////////////
 
 #if TRACE_TRIMMER_HAS_TRACE_WRITER_TESTS
-static int32 TryRunTests(int32 ArgC, TCHAR* ArgV[])
+static int32 TryRunTests(int32 ArgC, TCHAR const* const* ArgV)
 {
 	if (ArgC == 2 &&
 		(FCString::Stricmp(TEXT("-tests"), ArgV[1]) == 0 ||
@@ -109,7 +109,7 @@ static int32 TryRunTests(int32 ArgC, TCHAR* ArgV[])
 
 ////////////////////////////////////////////////////////////////////////////////////////////////////
 
-static int32 ParseCommandLine(FCmdLineParams& Parameters, int32 ArgC, TCHAR* ArgV[])
+static int32 ParseCommandLine(FCmdLineParams& Parameters, int32 ArgC, TCHAR const* const* ArgV)
 {
 	if (ArgC < 2)
 	{
@@ -366,13 +366,8 @@ TUniquePtr<UE::Trace::IOutDataStream> CreateOutputStream(FCmdLineParams& Paramet
 
 ////////////////////////////////////////////////////////////////////////////////////////////////////
 
-} // namespace UE::TraceTrimmer
-
-////////////////////////////////////////////////////////////////////////////////////////////////////
-
-INT32_MAIN_INT32_ARGC_TCHAR_ARGV()
+int32 Main(int32 ArgC, TCHAR const* const* ArgV)
 {
-	using namespace UE::TraceTrimmer;
 	int32 Result;
 
 #if TRACE_TRIMMER_HAS_TRACE_WRITER_TESTS
@@ -406,5 +401,22 @@ INT32_MAIN_INT32_ARGC_TCHAR_ARGV()
 	return TraceServices::Trim(Parameters.TrimParams, *InputDataStream, *OutputDataStream);
 }
 
+////////////////////////////////////////////////////////////////////////////////////////////////////
+
+} // namespace UE::TraceTrimmer
+
+////////////////////////////////////////////////////////////////////////////////////////////////////
+
+INT32_MAIN_INT32_ARGC_TCHAR_ARGV()
+{
+	int32 Result = UE::TraceTrimmer::Main(ArgC, ArgV);
+
+	// This ensures FUObjectArray::RemoveUObjectDeleteListener() is not called
+	// from FSparseDelegateStorage::FObjectListener::~FObjectListener().
+	RequestEngineExit(TEXT("TraceTrimmer completed"));
+
+	return Result;
+}
+
 ////////////////////////////////////////////////////////////////////////////////////////////////////
 ////////////////////////////////////////////////////////////////////////////////////////////////////
